[AI Toolsets]: Move DataFlowAgent toolset under the Toolsets directory

#Synced-CL 52371849

#Synced-CL 52371849

#Synced-CL 52371849

#Synced-CL 52371849

diff --git a/Engine/Plugins/Experimental/Toolsets/DataflowAgent/DataflowAgent.uplugin b/Engine/Plugins/Experimental/Toolsets/DataflowAgent/DataflowAgent.uplugin
new file mode 100644
--- /dev/null
+++ b/Engine/Plugins/Experimental/Toolsets/DataflowAgent/DataflowAgent.uplugin
@@ -0,0 +1,34 @@
+{
+	"FileVersion": 3,
+	"Version": 1,
+	"VersionName": "1.0",
+	"FriendlyName": "DataflowAgent",
+	"Description": "AI Agent toolset for editing Dataflow graphs",
+	"Category": "Other",
+	"EditorOnly": true,
+	"NoRedist": true,
+	"EnabledByDefault": false,
+	"CanContainContent": false,
+	"IsBetaVersion": false,
+	"IsExperimentalVersion": true,
+	"Installed": false,
+	"NoCode": false,
+	"Modules": [
+		{
+			"Name": "DataflowAgent",
+			"Type": "Editor",
+			"LoadingPhase": "PostEngineInit",
+			"TargetAllowList": ["Editor"]
+		}
+	],
+	"Plugins": [
+		{
+			"Name": "GeometryCollectionPlugin",
+			"Enabled": true
+		},
+		{
+			"Name": "ToolsetRegistry",
+			"Enabled": true
+		}
+	]
+}
diff --git a/Engine/Plugins/Experimental/Toolsets/DataflowAgent/Source/DataflowAgent/DataflowAgent.Build.cs b/Engine/Plugins/Experimental/Toolsets/DataflowAgent/Source/DataflowAgent/DataflowAgent.Build.cs
new file mode 100644
--- /dev/null
+++ b/Engine/Plugins/Experimental/Toolsets/DataflowAgent/Source/DataflowAgent/DataflowAgent.Build.cs
@@ -0,0 +1,42 @@
+// Copyright Epic Games, Inc. All Rights Reserved.
+
+using UnrealBuildTool;
+
+public class DataflowAgent : ModuleRules
+{
+	public DataflowAgent(ReadOnlyTargetRules Target) : base(Target)
+	{
+		PCHUsage = ModuleRules.PCHUsageMode.UseExplicitOrSharedPCHs;
+		bTreatAsEngineModule = true;
+
+		PublicDependencyModuleNames.AddRange(
+			new string[]
+			{
+				"Core",
+				"CoreUObject",
+				"Engine",
+				"ToolsetRegistry",
+				"DataflowCore",
+				"DataflowEngine",
+				"Projects",
+			}
+		);
+
+		PrivateDependencyModuleNames.AddRange(
+			new string[]
+			{
+				"CoreUObject",
+				"Engine",
+				"UnrealEd",
+				"DataflowEditor",
+				"AssetTools",
+				"AssetRegistry",
+				"Json",
+				"JsonUtilities",
+				"EditorSubsystem",
+				"GraphEditor",
+				"Kismet",
+			}
+		);
+	}
+}
diff --git a/Engine/Plugins/Experimental/Toolsets/DataflowAgent/Source/DataflowAgent/Private/DataflowAgentModule.cpp b/Engine/Plugins/Experimental/Toolsets/DataflowAgent/Source/DataflowAgent/Private/DataflowAgentModule.cpp
new file mode 100644
--- /dev/null
+++ b/Engine/Plugins/Experimental/Toolsets/DataflowAgent/Source/DataflowAgent/Private/DataflowAgentModule.cpp
@@ -0,0 +1,25 @@
+// Copyright Epic Games, Inc. All Rights Reserved.
+
+#include "DataflowAgentModule.h"
+
+#include "DataflowAgentToolset.h"
+#include "Misc/CoreDelegates.h"
+#include "ToolsetRegistry/UToolsetRegistry.h"
+
+DEFINE_LOG_CATEGORY(LogDataflowAgent);
+
+#define LOCTEXT_NAMESPACE "FDataflowAgentModule"
+
+void FDataflowAgentModule::StartupModule()
+{
+	UToolsetRegistry::RegisterToolsetClass(UDataflowAgentToolset::StaticClass());
+}
+
+void FDataflowAgentModule::ShutdownModule()
+{
+	UToolsetRegistry::UnregisterToolsetClass(UDataflowAgentToolset::StaticClass());
+}
+
+#undef LOCTEXT_NAMESPACE
+
+IMPLEMENT_MODULE(FDataflowAgentModule, DataflowAgent)
diff --git a/Engine/Plugins/Experimental/Toolsets/DataflowAgent/Source/DataflowAgent/Private/DataflowAgentToolset.cpp b/Engine/Plugins/Experimental/Toolsets/DataflowAgent/Source/DataflowAgent/Private/DataflowAgentToolset.cpp
new file mode 100644
--- /dev/null
+++ b/Engine/Plugins/Experimental/Toolsets/DataflowAgent/Source/DataflowAgent/Private/DataflowAgentToolset.cpp
@@ -0,0 +1,1111 @@
+// Copyright Epic Games, Inc. All Rights Reserved.
+
+#include "DataflowAgentToolset.h"
+
+#include "Dataflow/DataflowAssetEditUtils.h"
+#include "Dataflow/DataflowConnection.h"
+#include "Dataflow/DataflowEdNode.h"
+#include "Dataflow/DataflowGraph.h"
+#include "Dataflow/DataflowInputOutput.h"
+#include "Dataflow/DataflowNode.h"
+#include "Dataflow/DataflowNodeFactory.h"
+#include "Dataflow/DataflowObject.h"
+#include "Dataflow/DataflowVariableNodes.h"
+
+#include "AssetToolsModule.h"
+#include "Dom/JsonObject.h"
+#include "Dom/JsonValue.h"
+#include "Editor.h"
+#include "EdGraph/EdGraph.h"
+#include "EdGraph/EdGraphNode.h"
+#include "EdGraph/EdGraphPin.h"
+#include "EdGraphNode_Comment.h"
+#include "IAssetTools.h"
+#include "JsonObjectConverter.h"
+#include "Kismet/KismetSystemLibrary.h"
+#include "Misc/PackageName.h"
+#include "Policies/CondensedJsonPrintPolicy.h"
+#include "ScopedTransaction.h"
+#include "Serialization/JsonReader.h"
+#include "Serialization/JsonSerializer.h"
+#include "UObject/UnrealType.h"
+#include "UObject/Class.h"
+#include "UObject/UObjectIterator.h"
+
+// ---------------------------------------------------------------------------
+// Internal helpers
+// ---------------------------------------------------------------------------
+namespace DataflowAgentToolset
+{
+	/** Serialize any JSON value array to a compact string */
+	static FString JsonArrayToString(const TArray<TSharedPtr<FJsonValue>>& Array)
+	{
+		FString Out;
+		TSharedRef<TJsonWriter<TCHAR, TCondensedJsonPrintPolicy<TCHAR>>> Writer =
+			TJsonWriterFactory<TCHAR, TCondensedJsonPrintPolicy<TCHAR>>::Create(&Out);
+		FJsonSerializer::Serialize(Array, Writer);
+		return Out;
+	}
+
+	/** Serialize a JSON object to a compact string */
+	static FString JsonObjectToString(const TSharedPtr<FJsonObject>& Obj)
+	{
+		FString Out;
+		TSharedRef<TJsonWriter<TCHAR, TCondensedJsonPrintPolicy<TCHAR>>> Writer =
+			TJsonWriterFactory<TCHAR, TCondensedJsonPrintPolicy<TCHAR>>::Create(&Out);
+		FJsonSerializer::Serialize(Obj.ToSharedRef(), Writer);
+		return Out;
+	}
+
+	/** Try to parse a JSON object from a string; returns null on failure */
+	static TSharedPtr<FJsonObject> ParseJsonObject(const FString& JsonStr)
+	{
+		TSharedPtr<FJsonObject> Obj;
+		TSharedRef<TJsonReader<>> Reader = TJsonReaderFactory<>::Create(JsonStr);
+		FJsonSerializer::Deserialize(Reader, Obj);
+		return Obj;
+	}
+
+	/**
+	 * Find a UScriptStruct from a Dataflow type name like "FAddFloatsDataflowNode".
+	 * UScriptStructs are registered without the leading "F".
+	 */
+	static UScriptStruct* FindNodeStruct(FName TypeName)
+	{
+		FString StructName = TypeName.ToString();
+		// Strip leading "F" (USTRUCT convention)
+		if (StructName.Len() > 1 && StructName[0] == TEXT('F'))
+		{
+			StructName = StructName.RightChop(1);
+		}
+		for (TObjectIterator<UScriptStruct> It; It; ++It)
+		{
+			if ((*It)->GetName() == StructName)
+			{
+				return *It;
+			}
+		}
+		return nullptr;
+	}
+
+	/** Build a JSON array describing the pins of a DataflowNode */
+	static TArray<TSharedPtr<FJsonValue>> GetNodePinsJson(const FDataflowNode* Node, UE::Dataflow::FPin::EDirection Direction)
+	{
+		TArray<TSharedPtr<FJsonValue>> Result;
+		if (!Node) return Result;
+
+		for (const UE::Dataflow::FPin& Pin : Node->GetPins())
+		{
+			if (Pin.Direction != Direction) continue;
+			TSharedPtr<FJsonObject> PinObj = MakeShared<FJsonObject>();
+			PinObj->SetStringField(TEXT("name"), Pin.Name.ToString());
+			PinObj->SetStringField(TEXT("type"), Pin.Type.ToString());
+			Result.Add(MakeShared<FJsonValueObject>(PinObj));
+		}
+		return Result;
+	}
+
+	/** Build a JSON object describing a node's editable UPROPERTY fields */
+	static TArray<TSharedPtr<FJsonValue>> GetNodePropertiesJson(const FDataflowNode* Node, UScriptStruct* NodeStruct)
+	{
+		TArray<TSharedPtr<FJsonValue>> Result;
+		if (!Node || !NodeStruct) return Result;
+
+		// Walk up the struct hierarchy to collect all properties
+		for (TFieldIterator<FProperty> PropIt(NodeStruct, EFieldIteratorFlags::IncludeSuper); PropIt; ++PropIt)
+		{
+			const FProperty* Prop = *PropIt;
+			if (!Prop->HasAnyPropertyFlags(CPF_Edit)) continue;
+
+			// Skip the built-in base FDataflowNode properties (bOverrideColor, OverrideColor)
+			if (Prop->GetOwnerStruct() == FDataflowNode::StaticStruct()) continue;
+
+			TSharedPtr<FJsonObject> PropObj = MakeShared<FJsonObject>();
+			PropObj->SetStringField(TEXT("name"), Prop->GetName());
+			PropObj->SetStringField(TEXT("type"), Prop->GetCPPType());
+
+			FString DefaultValue;
+			Prop->ExportTextItem_Direct(DefaultValue, Prop->ContainerPtrToValuePtr<void>(Node), nullptr, nullptr, PPF_None);
+			PropObj->SetStringField(TEXT("default"), DefaultValue);
+
+			Result.Add(MakeShared<FJsonValueObject>(PropObj));
+		}
+		return Result;
+	}
+
+	/**
+	 * Apply JSON properties to a DataflowNode using UScriptStruct reflection.
+	 * Returns true if all specified properties were set successfully.
+	 */
+	static bool ApplyJsonToNode(FDataflowNode* Node, UScriptStruct* NodeStruct, const TSharedPtr<FJsonObject>& JsonParams)
+	{
+		if (!Node || !NodeStruct || !JsonParams.IsValid()) return false;
+
+		bool bAllSucceeded = true;
+		for (const auto& Pair : JsonParams->Values)
+		{
+			FProperty* Property = NodeStruct->FindPropertyByName(FName(*Pair.Key));
+			if (!Property)
+			{
+				// Search in parent structs
+				for (TFieldIterator<FProperty> PropIt(NodeStruct, EFieldIteratorFlags::IncludeSuper); PropIt; ++PropIt)
+				{
+					if ((*PropIt)->GetName().Equals(*Pair.Key, ESearchCase::IgnoreCase))
+					{
+						Property = *PropIt;
+						break;
+					}
+				}
+			}
+
+			if (!Property)
+			{
+				UE_LOGF(LogTemp, Warning, "DataflowAgentToolset: Property '%ls' not found on node type", *Pair.Key);
+				bAllSucceeded = false;
+				continue;
+			}
+
+			void* PropertyValue = Property->ContainerPtrToValuePtr<void>(Node);
+			if (!FJsonObjectConverter::JsonValueToUProperty(Pair.Value, Property, PropertyValue, 0, 0, false, nullptr))
+			{
+				UE_LOGF(LogTemp, Error, "DataflowAgentToolset : Failed to convert Json to property: %ls", *Pair.Key);
+				return false;
+			}
+
+			//FString ValueStr;
+			//if (Pair.Value->TryGetString(ValueStr))
+			//{
+			//	Property->ImportText_Direct(*ValueStr, Prop->ContainerPtrToValuePtr<void>(Node), nullptr, PPF_None);
+			//}
+			//else
+			//{
+			//	// For non-string JSON values, serialize to string first
+			//	TSharedRef<TJsonWriter<TCHAR, TCondensedJsonPrintPolicy<TCHAR>>> Writer =
+			//		TJsonWriterFactory<TCHAR, TCondensedJsonPrintPolicy<TCHAR>>::Create(&ValueStr);
+			//	FJsonSerializer::Serialize({ Pair.Value.ToSharedRef() }, Writer);
+			//	Property->ImportText_Direct(*ValueStr, Property->ContainerPtrToValuePtr<void>(Node), nullptr, PPF_None);
+			//}
+		}
+		return bAllSucceeded;
+	}
+
+	/**
+	 * Find a UClass from a name like "StaticMesh", "UStaticMesh", or "AMyActor".
+	 * UClasses are registered without the leading "U" or "A".
+	 */
+	static UClass* FindClass(FName ClassName)
+	{
+		FString Name = ClassName.ToString();
+		// Strip leading "U" or "A" (UCLASS convention)
+		if (Name.Len() > 1 && (Name[0] == TEXT('U') || Name[0] == TEXT('A')))
+		{
+			Name = Name.RightChop(1);
+		}
+		for (TObjectIterator<UClass> It; It; ++It)
+		{
+			if ((*It)->GetName() == Name)
+			{
+				return *It;
+			}
+		}
+		return nullptr;
+	}
+
+	/** Get the owning UDataflow asset from an EdNode */
+	static UDataflow* GetDataflowFromEdNode(const UDataflowEdNode* EdNode)
+	{
+		if (!EdNode) return nullptr;
+		return UDataflow::GetDataflowAssetFromEdGraph(EdNode->GetGraph());
+	}
+}
+
+// ---------------------------------------------------------------------------
+// Graph
+// ---------------------------------------------------------------------------
+
+FString UDataflowAgentToolset::CreateGraph(const FString& Name, const FString& Path)
+{
+	if (Name.IsEmpty())
+	{
+		UKismetSystemLibrary::RaiseScriptError(TEXT("Name parameter must be provided"));
+		return {};
+	}
+
+	const FString TargetPath = Path.IsEmpty() ? TEXT("/Game") : Path;
+	const FString FullPath = TargetPath / Name;
+
+	FAssetToolsModule& AssetToolsModule = FAssetToolsModule::GetModule();
+	FString UniqueAssetName;
+	FString PackageNameToUse;
+	AssetToolsModule.Get().CreateUniqueAssetName(FullPath, FString(), PackageNameToUse, UniqueAssetName);
+	const FString PackagePath = FPackageName::GetLongPackagePath(PackageNameToUse);
+
+	UObject* NewAsset = AssetToolsModule.Get().CreateAsset(UniqueAssetName, PackagePath, UDataflow::StaticClass(), nullptr);
+	UDataflow* Graph = Cast<UDataflow>(NewAsset);
+	if (!Graph)
+	{
+		UKismetSystemLibrary::RaiseScriptError(TEXT("Could not create Dataflow asset"));
+		return {};
+	}
+
+	Graph->MarkPackageDirty();
+	return Graph->GetPathName();
+}
+
+FString UDataflowAgentToolset::GetGraphStructure(const UDataflow* Graph)
+{
+	if (!Graph)
+	{
+		UKismetSystemLibrary::RaiseScriptError(TEXT("Error: No Dataflow asset provided"));
+		return {};
+	}
+
+	TSharedPtr<const UE::Dataflow::FGraph> DataflowGraph = Graph->GetDataflow();
+	if (!DataflowGraph.IsValid())
+	{
+		UKismetSystemLibrary::RaiseScriptError(TEXT("Error: Dataflow asset has no internal graph"));
+		return {};
+	}
+
+	TSharedPtr<FJsonObject> Root = MakeShared<FJsonObject>();
+	Root->SetStringField(TEXT("name"), Graph->GetFName().ToString());
+
+	// Nodes
+	TArray<TSharedPtr<FJsonValue>> NodesJson;
+	for (const TSharedPtr<FDataflowNode>& Node : DataflowGraph->GetNodes())
+	{
+		if (!Node.IsValid()) continue;
+		TSharedPtr<FJsonObject> NodeObj = MakeShared<FJsonObject>();
+		NodeObj->SetStringField(TEXT("name"), Node->GetName().ToString());
+		NodeObj->SetStringField(TEXT("type"), Node->GetType().ToString());
+		NodeObj->SetStringField(TEXT("guid"), Node->GetGuid().ToString());
+
+		// Position from the corresponding EdNode
+		if (TObjectPtr<const UDataflowEdNode> EdNode = Graph->FindEdNodeByDataflowNodeGuid(Node->GetGuid()))
+		{
+			NodeObj->SetNumberField(TEXT("posX"), EdNode->NodePosX);
+			NodeObj->SetNumberField(TEXT("posY"), EdNode->NodePosY);
+		}
+
+		// Pins
+		NodeObj->SetArrayField(TEXT("inputPins"),
+			DataflowAgentToolset::GetNodePinsJson(Node.Get(), UE::Dataflow::FPin::EDirection::INPUT));
+		NodeObj->SetArrayField(TEXT("outputPins"),
+			DataflowAgentToolset::GetNodePinsJson(Node.Get(), UE::Dataflow::FPin::EDirection::OUTPUT));
+
+		NodesJson.Add(MakeShared<FJsonValueObject>(NodeObj));
+	}
+	Root->SetArrayField(TEXT("nodes"), NodesJson);
+
+	// Connections
+	TArray<TSharedPtr<FJsonValue>> ConnectionsJson;
+	for (const UE::Dataflow::FLink& Link : DataflowGraph->GetConnections())
+	{
+		TSharedPtr<FJsonObject> LinkObj = MakeShared<FJsonObject>();
+
+		// Resolve node guids to names where possible
+		if (TSharedPtr<const FDataflowNode> FromNode = DataflowGraph->FindBaseNode(Link.OutputNode))
+		{
+			LinkObj->SetStringField(TEXT("fromNode"), FromNode->GetName().ToString());
+		}
+		else
+		{
+			LinkObj->SetStringField(TEXT("fromNodeGuid"), Link.OutputNode.ToString());
+		}
+		LinkObj->SetStringField(TEXT("fromPinGuid"), Link.Output.ToString());
+
+		if (TSharedPtr<const FDataflowNode> ToNode = DataflowGraph->FindBaseNode(Link.InputNode))
+		{
+			LinkObj->SetStringField(TEXT("toNode"), ToNode->GetName().ToString());
+		}
+		else
+		{
+			LinkObj->SetStringField(TEXT("toNodeGuid"), Link.InputNode.ToString());
+		}
+		LinkObj->SetStringField(TEXT("toPinGuid"), Link.Input.ToString());
+
+		// Resolve pin names from the connection objects
+		if (TSharedPtr<const FDataflowNode> FromNode = DataflowGraph->FindBaseNode(Link.OutputNode))
+		{
+			if (const FDataflowOutput* Output = FromNode->FindOutput(Link.Output))
+			{
+				LinkObj->SetStringField(TEXT("fromPin"), Output->GetName().ToString());
+			}
+		}
+		if (TSharedPtr<const FDataflowNode> ToNode = DataflowGraph->FindBaseNode(Link.InputNode))
+		{
+			if (const FDataflowInput* Input = ToNode->FindInput(Link.Input))
+			{
+				LinkObj->SetStringField(TEXT("toPin"), Input->GetName().ToString());
+			}
+		}
+
+		ConnectionsJson.Add(MakeShared<FJsonValueObject>(LinkObj));
+	}
+	Root->SetArrayField(TEXT("connections"), ConnectionsJson);
+
+	return DataflowAgentToolset::JsonObjectToString(Root);
+}
+
+// ---------------------------------------------------------------------------
+// Node Types
+// ---------------------------------------------------------------------------
+
+FString UDataflowAgentToolset::ListNodeTypes(const bool bCommonOnly)
+{
+	UE::Dataflow::FNodeFactory* Factory = UE::Dataflow::FNodeFactory::GetInstance();
+	if (!Factory)
+	{
+		UKismetSystemLibrary::RaiseScriptError(TEXT("Error: Dataflow NodeFactory not available"));
+		return {};
+	}
+
+	TArray<TSharedPtr<FJsonValue>> Result;
+	for (const UE::Dataflow::FFactoryParameters& Params : Factory->RegisteredParameters())
+	{
+		if (!Params.IsValid()) continue;
+		if (bCommonOnly && (Params.IsDeprecated() || Params.IsExperimental())) continue;
+
+		TSharedPtr<FJsonObject> NodeObj = MakeShared<FJsonObject>();
+		NodeObj->SetStringField(TEXT("typeName"),    Params.TypeName.ToString());
+		NodeObj->SetStringField(TEXT("displayName"), Params.DisplayName.ToString());
+		NodeObj->SetStringField(TEXT("category"),    Params.Category.ToString());
+		if (!Params.ToolTip.IsEmpty())
+		{
+			NodeObj->SetStringField(TEXT("tooltip"), Params.ToolTip);
+		}
+		NodeObj->SetBoolField(TEXT("isDeprecated"), Params.bIsDeprecated);
+		NodeObj->SetBoolField(TEXT("isExperimental"), Params.bIsExperimental);
+
+		Result.Add(MakeShared<FJsonValueObject>(NodeObj));
+	}
+
+	return DataflowAgentToolset::JsonArrayToString(Result);
+}
+
+FString UDataflowAgentToolset::GetNodeTypeSchema(const FString& TypeName)
+{
+	if (TypeName.IsEmpty())
+	{
+		UKismetSystemLibrary::RaiseScriptError(TEXT("Error: TypeName must be provided"));
+		return {};
+	}
+
+	UE::Dataflow::FNodeFactory* Factory = UE::Dataflow::FNodeFactory::GetInstance();
+	if (!Factory)
+	{
+		UKismetSystemLibrary::RaiseScriptError(TEXT("Error: Dataflow NodeFactory not available"));
+		return {};
+	}
+
+	const UE::Dataflow::FFactoryParameters& Params = Factory->GetParameters(FName(*TypeName));
+	if (!Params.IsValid())
+	{
+		UKismetSystemLibrary::RaiseScriptError(FString::Printf(TEXT("Error: Node type '%s' not found"), *TypeName));
+		return {};
+	}
+
+	TSharedPtr<FJsonObject> Schema = MakeShared<FJsonObject>();
+	Schema->SetStringField(TEXT("typeName"),    Params.TypeName.ToString());
+	Schema->SetStringField(TEXT("displayName"), Params.DisplayName.ToString());
+	Schema->SetStringField(TEXT("category"),    Params.Category.ToString());
+	Schema->SetStringField(TEXT("tooltip"),     Params.ToolTip);
+	Schema->SetBoolField(TEXT("isDeprecated"), Params.bIsDeprecated);
+	Schema->SetBoolField(TEXT("isExperimental"), Params.bIsExperimental);
+
+	const FDataflowNode* DefaultNode = Params.DefaultNodeObject.Get();
+	Schema->SetArrayField(TEXT("inputPins"),
+		DataflowAgentToolset::GetNodePinsJson(DefaultNode, UE::Dataflow::FPin::EDirection::INPUT));
+	Schema->SetArrayField(TEXT("outputPins"),
+		DataflowAgentToolset::GetNodePinsJson(DefaultNode, UE::Dataflow::FPin::EDirection::OUTPUT));
+
+	// Editable properties
+	UScriptStruct* NodeStruct = DataflowAgentToolset::FindNodeStruct(Params.TypeName);
+	if (NodeStruct && DefaultNode)
+	{
+		Schema->SetArrayField(TEXT("properties"),
+			DataflowAgentToolset::GetNodePropertiesJson(DefaultNode, NodeStruct));
+	}
+
+	return DataflowAgentToolset::JsonObjectToString(Schema);
+}
+
+// ---------------------------------------------------------------------------
+// Nodes
+// ---------------------------------------------------------------------------
+
+UDataflowEdNode* UDataflowAgentToolset::AddNode(
+	UDataflow* Graph,
+	const FString& TypeName,
+	const FString& NodeName,
+	const FString& JsonParams,
+	const int32 X,
+	const int32 Y)
+{
+	if (!Graph)
+	{
+		UKismetSystemLibrary::RaiseScriptError(TEXT("Error: No Dataflow asset provided"));
+		return nullptr;
+	}
+	if (TypeName.IsEmpty())
+	{
+		UKismetSystemLibrary::RaiseScriptError(TEXT("Error: TypeName must be provided"));
+		return nullptr;
+	}
+	if (NodeName.IsEmpty())
+	{
+		UKismetSystemLibrary::RaiseScriptError(TEXT("Error: NodeName must be provided"));
+		return nullptr;
+	}
+
+	UE::Dataflow::FNodeFactory* Factory = UE::Dataflow::FNodeFactory::GetInstance();
+	if (!Factory)
+	{
+		UKismetSystemLibrary::RaiseScriptError(TEXT("Error: Dataflow NodeFactory not available"));
+		return nullptr;
+	}
+
+	if (!Factory->GetParameters(FName(*TypeName)).IsValid())
+	{
+		UKismetSystemLibrary::RaiseScriptError(FString::Printf(TEXT("Error: Node type '%s' not found"), *TypeName));
+		return nullptr;
+	}
+
+	FScopedTransaction Transaction(INVTEXT("AddDataflowNodeToolCall"));
+
+	UDataflowEdNode* NewEdNode = UE::Dataflow::FEditAssetUtils::AddNewNode(
+		Graph, FVector2D(X, Y), FName(*NodeName), FName(*TypeName), /*FromPin=*/nullptr);
+
+	if (!NewEdNode)
+	{
+		Transaction.Cancel();
+		UKismetSystemLibrary::RaiseScriptError(FString::Printf(
+			TEXT("Error: Failed to create node of type '%s'"), *TypeName));
+		return nullptr;
+	}
+
+	// Apply JSON parameters if provided
+	if (!JsonParams.IsEmpty())
+	{
+		TSharedPtr<FJsonObject> ParamsObj = DataflowAgentToolset::ParseJsonObject(JsonParams);
+		if (!ParamsObj.IsValid())
+		{
+			Transaction.Cancel();
+			UKismetSystemLibrary::RaiseScriptError(TEXT("Error: Failed to parse JsonParams"));
+			return nullptr;
+		}
+
+		TSharedPtr<FDataflowNode> DataflowNode = NewEdNode->GetDataflowNode();
+		UScriptStruct* NodeStruct = DataflowAgentToolset::FindNodeStruct(FName(*TypeName));
+		if (NodeStruct && DataflowNode.IsValid())
+		{
+			DataflowAgentToolset::ApplyJsonToNode(DataflowNode.Get(), NodeStruct, ParamsObj);
+		}
+	}
+
+	// Special case for variables nodes as we need to refresh the node to properly connect to the existing variable
+	if (TSharedPtr<FDataflowNode> DataflowNode = NewEdNode->GetDataflowNode())
+	{
+		if (FGetDataflowVariableNode* VariableNode = DataflowNode->AsType<FGetDataflowVariableNode>())
+		{
+			VariableNode->SetVariable(Graph, VariableNode->GetVariableName());
+		}
+	}
+
+	Graph->MarkPackageDirty();
+	return NewEdNode;
+}
+
+void UDataflowAgentToolset::UpdateNode(UDataflowEdNode* Node, const FString& JsonParams)
+{
+	if (!Node)
+	{
+		UKismetSystemLibrary::RaiseScriptError(TEXT("Error: No EdNode provided"));
+		return;
+	}
+	if (JsonParams.IsEmpty())
+	{
+		return;
+	}
+
+	TSharedPtr<FDataflowNode> DataflowNode = Node->GetDataflowNode();
+	if (!DataflowNode.IsValid())
+	{
+		UKismetSystemLibrary::RaiseScriptError(TEXT("Error: EdNode has no associated Dataflow node"));
+		return;
+	}
+
+	TSharedPtr<FJsonObject> ParamsObj = DataflowAgentToolset::ParseJsonObject(JsonParams);
+	if (!ParamsObj.IsValid())
+	{
+		UKismetSystemLibrary::RaiseScriptError(TEXT("Error: Failed to parse JsonParams"));
+		return;
+	}
+
+	UScriptStruct* NodeStruct = DataflowAgentToolset::FindNodeStruct(DataflowNode->GetType());
+	if (!NodeStruct)
+	{
+		UKismetSystemLibrary::RaiseScriptError(FString::Printf(
+			TEXT("Error: Could not find UScriptStruct for node type '%s'"),
+			*DataflowNode->GetType().ToString()));
+		return;
+	}
+
+	FScopedTransaction Transaction(INVTEXT("UpdateDataflowNodeToolCall"));
+	Node->Modify();
+
+	if (!DataflowAgentToolset::ApplyJsonToNode(DataflowNode.Get(), NodeStruct, ParamsObj))
+	{
+		Transaction.Cancel();
+		UKismetSystemLibrary::RaiseScriptError(TEXT("Error: One or more properties could not be set"));
+		return;
+	}
+
+	UDataflow* DataflowAsset = DataflowAgentToolset::GetDataflowFromEdNode(Node);
+	if (DataflowAsset)
+	{
+		DataflowAsset->RefreshEdNode(Node);
+		DataflowAsset->MarkPackageDirty();
+	}
+}
+
+FString UDataflowAgentToolset::GetNodeInfo(UDataflowEdNode* Node)
+{
+	if (!Node)
+	{
+		UKismetSystemLibrary::RaiseScriptError(TEXT("Error: No EdNode provided"));
+		return {};
+	}
+
+	TSharedPtr<const FDataflowNode> DataflowNode = Node->GetDataflowNode();
+
+	TSharedPtr<FJsonObject> Info = MakeShared<FJsonObject>();
+	Info->SetStringField(TEXT("name"),    Node->GetName());
+	Info->SetStringField(TEXT("guid"),    Node->GetDataflowNodeGuid().ToString());
+	Info->SetNumberField(TEXT("posX"),    Node->NodePosX);
+	Info->SetNumberField(TEXT("posY"),    Node->NodePosY);
+
+	if (DataflowNode.IsValid())
+	{
+		Info->SetStringField(TEXT("type"),        DataflowNode->GetType().ToString());
+		Info->SetStringField(TEXT("displayName"), DataflowNode->GetDisplayName().ToString());
+		Info->SetStringField(TEXT("category"),    DataflowNode->GetCategory().ToString());
+
+		Info->SetArrayField(TEXT("inputPins"),
+			DataflowAgentToolset::GetNodePinsJson(DataflowNode.Get(), UE::Dataflow::FPin::EDirection::INPUT));
+		Info->SetArrayField(TEXT("outputPins"),
+			DataflowAgentToolset::GetNodePinsJson(DataflowNode.Get(), UE::Dataflow::FPin::EDirection::OUTPUT));
+
+		// Current property values
+		UScriptStruct* NodeStruct = DataflowAgentToolset::FindNodeStruct(DataflowNode->GetType());
+		if (NodeStruct)
+		{
+			TSharedPtr<FJsonObject> Props = MakeShared<FJsonObject>();
+			for (TFieldIterator<FProperty> PropIt(NodeStruct, EFieldIteratorFlags::IncludeSuper); PropIt; ++PropIt)
+			{
+				const FProperty* Prop = *PropIt;
+				if (!Prop->HasAnyPropertyFlags(CPF_Edit)) continue;
+				if (Prop->GetOwnerStruct() == FDataflowNode::StaticStruct()) continue;
+
+				FString Value;
+				Prop->ExportTextItem_Direct(Value, Prop->ContainerPtrToValuePtr<void>(DataflowNode.Get()), nullptr, nullptr, PPF_None);
+				Props->SetStringField(Prop->GetName(), Value);
+			}
+			Info->SetObjectField(TEXT("properties"), Props);
+		}
+	}
+
+	return DataflowAgentToolset::JsonObjectToString(Info);
+}
+
+void UDataflowAgentToolset::RepositionNode(UDataflowEdNode* Node, const int32 X, const int32 Y)
+{
+	if (!Node)
+	{
+		UKismetSystemLibrary::RaiseScriptError(TEXT("Error: No EdNode provided"));
+		return;
+	}
+
+	FScopedTransaction Transaction(INVTEXT("RepositionDataflowNodeToolCall"));
+	Node->Modify();
+	Node->NodePosX = X;
+	Node->NodePosY = Y;
+
+	UDataflow* DataflowAsset = DataflowAgentToolset::GetDataflowFromEdNode(Node);
+	if (DataflowAsset)
+	{
+		DataflowAsset->MarkPackageDirty();
+	}
+}
+
+void UDataflowAgentToolset::RemoveNode(UDataflow* Graph, UDataflowEdNode* Node)
+{
+	if (!Graph)
+	{
+		UKismetSystemLibrary::RaiseScriptError(TEXT("Error: No Dataflow asset provided"));
+		return;
+	}
+	if (!Node)
+	{
+		UKismetSystemLibrary::RaiseScriptError(TEXT("Error: No EdNode provided"));
+		return;
+	}
+
+	// Verify the node belongs to this graph
+	if (UDataflow::GetDataflowAssetFromEdGraph(Node->GetGraph()) != Graph)
+	{
+		UKismetSystemLibrary::RaiseScriptError(TEXT("Error: Node does not belong to the provided graph"));
+		return;
+	}
+
+	TArray<UEdGraphNode*> NodesToDelete = { Node };
+	UE::Dataflow::FEditAssetUtils::DeleteNodes(Graph, NodesToDelete);
+}
+
+// ---------------------------------------------------------------------------
+// Connections
+// ---------------------------------------------------------------------------
+
+bool UDataflowAgentToolset::ConnectNodePins(
+	UDataflowEdNode* FromNode, const FString& FromPin,
+	UDataflowEdNode* ToNode,   const FString& ToPin)
+{
+	if (!FromNode)
+	{
+		UKismetSystemLibrary::RaiseScriptError(TEXT("Error: FromNode is null"));
+		return false;
+	}
+	if (!ToNode)
+	{
+		UKismetSystemLibrary::RaiseScriptError(TEXT("Error: ToNode is null"));
+		return false;
+	}
+
+	UDataflow* FromAsset = DataflowAgentToolset::GetDataflowFromEdNode(FromNode);
+	UDataflow* ToAsset   = DataflowAgentToolset::GetDataflowFromEdNode(ToNode);
+	if (!FromAsset || FromAsset != ToAsset)
+	{
+		UKismetSystemLibrary::RaiseScriptError(TEXT("Error: Both nodes must belong to the same Dataflow graph"));
+		return false;
+	}
+
+	TSharedPtr<FDataflowNode> SourceDFNode = FromNode->GetDataflowNode();
+	TSharedPtr<FDataflowNode> TargetDFNode = ToNode->GetDataflowNode();
+	if (!SourceDFNode.IsValid() || !TargetDFNode.IsValid())
+	{
+		UKismetSystemLibrary::RaiseScriptError(TEXT("Error: Could not retrieve internal Dataflow nodes"));
+		return false;
+	}
+
+	FDataflowOutput* Output = SourceDFNode->FindOutput(FName(*FromPin));
+	if (!Output)
+	{
+		UKismetSystemLibrary::RaiseScriptError(FString::Printf(
+			TEXT("Error: Node '%s' has no output pin named '%s'"),
+			*FromNode->GetName(), *FromPin));
+		return false;
+	}
+
+	FDataflowInput* Input = TargetDFNode->FindInput(FName(*ToPin));
+	if (!Input)
+	{
+		UKismetSystemLibrary::RaiseScriptError(FString::Printf(
+			TEXT("Error: Node '%s' has no input pin named '%s'"),
+			*ToNode->GetName(), *ToPin));
+		return false;
+	}
+
+	TSharedPtr<UE::Dataflow::FGraph> DataflowGraph = FromAsset->GetDataflow();
+	if (!DataflowGraph.IsValid())
+	{
+		UKismetSystemLibrary::RaiseScriptError(TEXT("Error: Dataflow graph is invalid"));
+		return false;
+	}
+
+	if (!DataflowGraph->CanConnect(*Output, *Input))
+	{
+		UKismetSystemLibrary::RaiseScriptError(FString::Printf(
+			TEXT("Error: Cannot connect pin '%s' on '%s' to pin '%s' on '%s' (type incompatibility)"),
+			*FromPin, *FromNode->GetName(), *ToPin, *ToNode->GetName()));
+		return false;
+	}
+
+	FScopedTransaction Transaction(INVTEXT("ConnectDataflowPinsToolCall"));
+	FromAsset->Modify();
+
+	const bool bConnected = DataflowGraph->Connect(*Output, *Input);
+	if (!bConnected)
+	{
+		Transaction.Cancel();
+		UKismetSystemLibrary::RaiseScriptError(FString::Printf(
+			TEXT("Error: Failed to connect '%s'.'%s' -> '%s'.'%s'"),
+			*FromNode->GetName(), *FromPin, *ToNode->GetName(), *ToPin));
+		return false;
+	}
+
+	// Sync the visual EdGraph pins
+	FromAsset->RefreshEdNode(FromNode);
+	FromAsset->RefreshEdNode(ToNode);
+	FromAsset->MarkPackageDirty();
+
+	return true;
+}
+
+void UDataflowAgentToolset::DisconnectNodePins(
+	UDataflowEdNode* FromNode, const FString& FromPin,
+	UDataflowEdNode* ToNode,   const FString& ToPin)
+{
+	if (!FromNode)
+	{
+		UKismetSystemLibrary::RaiseScriptError(TEXT("Error: FromNode is null"));
+		return;
+	}
+	if (!ToNode)
+	{
+		UKismetSystemLibrary::RaiseScriptError(TEXT("Error: ToNode is null"));
+		return;
+	}
+
+	UDataflow* FromAsset = DataflowAgentToolset::GetDataflowFromEdNode(FromNode);
+	UDataflow* ToAsset   = DataflowAgentToolset::GetDataflowFromEdNode(ToNode);
+	if (!FromAsset || FromAsset != ToAsset)
+	{
+		UKismetSystemLibrary::RaiseScriptError(TEXT("Error: Both nodes must belong to the same Dataflow graph"));
+		return;
+	}
+
+	TSharedPtr<FDataflowNode> SourceDFNode = FromNode->GetDataflowNode();
+	TSharedPtr<FDataflowNode> TargetDFNode = ToNode->GetDataflowNode();
+	if (!SourceDFNode.IsValid() || !TargetDFNode.IsValid())
+	{
+		UKismetSystemLibrary::RaiseScriptError(TEXT("Error: Could not retrieve internal Dataflow nodes"));
+		return;
+	}
+
+	FDataflowOutput* Output = SourceDFNode->FindOutput(FName(*FromPin));
+	if (!Output)
+	{
+		UKismetSystemLibrary::RaiseScriptError(FString::Printf(
+			TEXT("Error: Node '%s' has no output pin named '%s'"),
+			*FromNode->GetName(), *FromPin));
+		return;
+	}
+
+	FDataflowInput* Input = TargetDFNode->FindInput(FName(*ToPin));
+	if (!Input)
+	{
+		UKismetSystemLibrary::RaiseScriptError(FString::Printf(
+			TEXT("Error: Node '%s' has no input pin named '%s'"),
+			*ToNode->GetName(), *ToPin));
+		return;
+	}
+
+	TSharedPtr<UE::Dataflow::FGraph> DataflowGraph = FromAsset->GetDataflow();
+	if (!DataflowGraph.IsValid())
+	{
+		return;
+	}
+
+	FScopedTransaction Transaction(INVTEXT("DisconnectDataflowPinsToolCall"));
+	FromAsset->Modify();
+
+	DataflowGraph->Disconnect(Output, Input);
+
+	// Sync the visual EdGraph pins
+	FromAsset->RefreshEdNode(FromNode);
+	FromAsset->RefreshEdNode(ToNode);
+	FromAsset->MarkPackageDirty();
+}
+
+// ---------------------------------------------------------------------------
+// Variables
+// ---------------------------------------------------------------------------
+
+FString UDataflowAgentToolset::ListVariables(const UDataflow* Graph)
+{
+	if (!Graph)
+	{
+		UKismetSystemLibrary::RaiseScriptError(TEXT("Error: No Dataflow asset provided"));
+		return {};
+	}
+
+	const UPropertyBag* BagStruct = Graph->Variables.GetPropertyBagStruct();
+	if (!BagStruct)
+	{
+		return DataflowAgentToolset::JsonArrayToString({});
+	}
+
+	TArray<TSharedPtr<FJsonValue>> Result;
+	for (const FPropertyBagPropertyDesc& Desc : BagStruct->GetPropertyDescs())
+	{
+		TSharedPtr<FJsonObject> VarObj = MakeShared<FJsonObject>();
+		VarObj->SetStringField(TEXT("name"), Desc.Name.ToString());
+
+		const UEnum* TypeEnum = StaticEnum<EPropertyBagPropertyType>();
+		VarObj->SetStringField(TEXT("type"), TypeEnum ? TypeEnum->GetNameStringByValue(static_cast<int64>(Desc.ValueType)) : FString::FromInt(static_cast<int32>(Desc.ValueType)));
+
+		if (Desc.ValueTypeObject)
+		{
+			VarObj->SetStringField(TEXT("valueTypeName"), Desc.ValueTypeObject->GetName());
+		}
+
+		TValueOrError<FString, EPropertyBagResult> ValueResult = Graph->Variables.GetValueSerializedString(Desc.Name);
+		if (ValueResult.HasValue())
+		{
+			VarObj->SetStringField(TEXT("value"), ValueResult.GetValue());
+		}
+
+		Result.Add(MakeShared<FJsonValueObject>(VarObj));
+	}
+
+	return DataflowAgentToolset::JsonArrayToString(Result);
+}
+
+bool UDataflowAgentToolset::AddVariable(UDataflow* Graph, const FString& Name, const FString& Type)
+{
+	if (!Graph)
+	{
+		UKismetSystemLibrary::RaiseScriptError(TEXT("Error: No Dataflow asset provided"));
+		return false;
+	}
+	if (Name.IsEmpty())
+	{
+		UKismetSystemLibrary::RaiseScriptError(TEXT("Error: Name must be provided"));
+		return false;
+	}
+	if (Type.IsEmpty())
+	{
+		UKismetSystemLibrary::RaiseScriptError(TEXT("Error: Type must be provided"));
+		return false;
+	}
+
+	static const TMap<FString, EPropertyBagPropertyType> PrimitiveTypeMap =
+	{
+		{ TEXT("Bool"),   EPropertyBagPropertyType::Bool   },
+		{ TEXT("Int32"),  EPropertyBagPropertyType::Int32  },
+		{ TEXT("Int64"),  EPropertyBagPropertyType::Int64  },
+		{ TEXT("Float"),  EPropertyBagPropertyType::Float  },
+		{ TEXT("Double"), EPropertyBagPropertyType::Double },
+		{ TEXT("Name"),   EPropertyBagPropertyType::Name   },
+		{ TEXT("String"), EPropertyBagPropertyType::String },
+	};
+
+	EPropertyBagPropertyType BagType  = EPropertyBagPropertyType::None;
+	const UObject*           TypeObj  = nullptr;
+
+	if (const EPropertyBagPropertyType* Primitive = PrimitiveTypeMap.Find(Type))
+	{
+		BagType = *Primitive;
+	}
+	else if (Type.StartsWith(TEXT("Object:")))
+	{
+		// "Object:ClassName" or "Object:UMyClass"
+		const FString ClassName = Type.RightChop(7); // strip "Object:"
+		UClass* FoundClass = DataflowAgentToolset::FindClass(FName(*ClassName));
+		if (!FoundClass)
+		{
+			UKismetSystemLibrary::RaiseScriptError(FString::Printf(
+				TEXT("Error: Could not find UClass '%s'"), *ClassName));
+			return false;
+		}
+		BagType = EPropertyBagPropertyType::Object;
+		TypeObj = FoundClass;
+	}
+	else
+	{
+		// Treat as UScriptStruct name (e.g. "Vector", "FVector", "Transform", "FTransform")
+		UScriptStruct* FoundStruct = DataflowAgentToolset::FindNodeStruct(FName(*Type));
+		if (!FoundStruct)
+		{
+			UKismetSystemLibrary::RaiseScriptError(FString::Printf(
+				TEXT("Error: Unknown type '%s'. For primitives use: Bool, Int32, Int64, Float, Double, Name, String. ")
+				TEXT("For structs pass the struct name (e.g. Vector, Transform, Rotator). ")
+				TEXT("For objects use 'Object:ClassName'."), *Type));
+			return false;
+		}
+		BagType = EPropertyBagPropertyType::Struct;
+		TypeObj = FoundStruct;
+	}
+
+	const FName VarName(*Name);
+	if (Graph->Variables.FindPropertyDescByName(VarName) != nullptr)
+	{
+		UKismetSystemLibrary::RaiseScriptError(FString::Printf(
+			TEXT("Error: Variable '%s' already exists on this graph"), *Name));
+		return false;
+	}
+
+	FScopedTransaction Transaction(INVTEXT("AddDataflowVariableToolCall"));
+	Graph->Modify();
+
+	Graph->Variables.AddProperties({ FPropertyBagPropertyDesc(VarName, BagType, TypeObj) });
+	Graph->MarkPackageDirty();
+	return true;
+}
+
+void UDataflowAgentToolset::RemoveVariable(UDataflow* Graph, const FString& Name)
+{
+	if (!Graph)
+	{
+		UKismetSystemLibrary::RaiseScriptError(TEXT("Error: No Dataflow asset provided"));
+		return;
+	}
+	if (Name.IsEmpty())
+	{
+		UKismetSystemLibrary::RaiseScriptError(TEXT("Error: Name must be provided"));
+		return;
+	}
+
+	const FName VarName(*Name);
+	if (Graph->Variables.FindPropertyDescByName(VarName) == nullptr)
+	{
+		UKismetSystemLibrary::RaiseScriptError(FString::Printf(
+			TEXT("Error: Variable '%s' not found on this graph"), *Name));
+		return;
+	}
+
+	FScopedTransaction Transaction(INVTEXT("RemoveDataflowVariableToolCall"));
+	Graph->Modify();
+
+	Graph->Variables.RemovePropertyByName(VarName);
+	Graph->MarkPackageDirty();
+}
+
+bool UDataflowAgentToolset::SetVariable(UDataflow* Graph, const FString& Name, const FString& Value)
+{
+	if (!Graph)
+	{
+		UKismetSystemLibrary::RaiseScriptError(TEXT("Error: No Dataflow asset provided"));
+		return false;
+	}
+	if (Name.IsEmpty())
+	{
+		UKismetSystemLibrary::RaiseScriptError(TEXT("Error: Name must be provided"));
+		return false;
+	}
+
+	const FName VarName(*Name);
+	if (Graph->Variables.FindPropertyDescByName(VarName) == nullptr)
+	{
+		UKismetSystemLibrary::RaiseScriptError(FString::Printf(
+			TEXT("Error: Variable '%s' not found on this graph"), *Name));
+		return false;
+	}
+
+	FScopedTransaction Transaction(INVTEXT("SetDataflowVariableToolCall"));
+	Graph->Modify();
+
+	const EPropertyBagResult SetResult = Graph->Variables.SetValueSerializedString(VarName, Value);
+	if (SetResult != EPropertyBagResult::Success)
+	{
+		Transaction.Cancel();
+		UKismetSystemLibrary::RaiseScriptError(FString::Printf(
+			TEXT("Error: Failed to set variable '%s' to '%s'"), *Name, *Value));
+		return false;
+	}
+
+	Graph->MarkPackageDirty();
+	return true;
+}
+
+// ---------------------------------------------------------------------------
+// Comments
+// ---------------------------------------------------------------------------
+
+FString UDataflowAgentToolset::AddCommentBox(
+	UDataflow* Graph,
+	const TArray<UDataflowEdNode*>& Nodes,
+	const FString& Comment,
+	FLinearColor Color)
+{
+	if (!Graph)
+	{
+		UKismetSystemLibrary::RaiseScriptError(TEXT("Error: No Dataflow asset provided"));
+		return {};
+	}
+	if (Nodes.IsEmpty())
+	{
+		UKismetSystemLibrary::RaiseScriptError(TEXT("Error: Nodes array must not be empty"));
+		return {};
+	}
+
+	// Compute bounding box around provided nodes
+	int32 MinX = INT_MAX, MinY = INT_MAX, MaxX = INT_MIN, MaxY = INT_MIN;
+	for (const UDataflowEdNode* Node : Nodes)
+	{
+		if (!Node) continue;
+		const int32 NodeW = 200;
+		const int32 NodeH = 100;
+		MinX = FMath::Min(MinX, Node->NodePosX);
+		MinY = FMath::Min(MinY, Node->NodePosY);
+		MaxX = FMath::Max(MaxX, Node->NodePosX + NodeW);
+		MaxY = FMath::Max(MaxY, Node->NodePosY + NodeH);
+	}
+	if (MinX > MaxX || MinY > MaxY)
+	{
+		UKismetSystemLibrary::RaiseScriptError(TEXT("Error: Invalid bounds for nodes"));
+		return {};
+	}
+
+	const int32 Padding = 50;
+	MinX -= Padding; MinY -= Padding;
+	MaxX += Padding; MaxY += Padding;
+
+	FScopedTransaction Transaction(INVTEXT("AddDataflowCommentBoxToolCall"));
+
+	UEdGraphNode* CommentBase = UE::Dataflow::FEditAssetUtils::AddNewComment(
+		Graph, FVector2D(MinX, MinY), FVector2D(MaxX - MinX, MaxY - MinY));
+
+	if (UEdGraphNode_Comment* CommentNode = Cast<UEdGraphNode_Comment>(CommentBase))
+	{
+		CommentNode->Modify();
+		CommentNode->NodeComment = Comment;
+		CommentNode->CommentColor = Color;
+		Graph->MarkPackageDirty();
+		return CommentNode->NodeGuid.ToString();
+	}
+
+	Transaction.Cancel();
+	UKismetSystemLibrary::RaiseScriptError(TEXT("Error: Failed to create comment node"));
+	return {};
+}
+
+void UDataflowAgentToolset::RemoveCommentBox(UDataflow* Graph, const FString& CommentId)
+{
+	if (!Graph)
+	{
+		UKismetSystemLibrary::RaiseScriptError(TEXT("Error: No Dataflow asset provided"));
+		return;
+	}
+	if (CommentId.IsEmpty())
+	{
+		UKismetSystemLibrary::RaiseScriptError(TEXT("Error: CommentId must be provided"));
+		return;
+	}
+
+	// Find the comment node by GUID
+	UEdGraphNode_Comment* FoundComment = nullptr;
+	for (UEdGraphNode* Node : Graph->Nodes)
+	{
+		if (UEdGraphNode_Comment* Comment = Cast<UEdGraphNode_Comment>(Node))
+		{
+			if (Comment->NodeGuid.ToString() == CommentId)
+			{
+				FoundComment = Comment;
+				break;
+			}
+		}
+	}
+
+	if (!FoundComment)
+	{
+		UKismetSystemLibrary::RaiseScriptError(FString::Printf(
+			TEXT("Error: No comment node found with ID '%s'"), *CommentId));
+		return;
+	}
+
+	FScopedTransaction Transaction(INVTEXT("RemoveDataflowCommentBoxToolCall"));
+	Graph->Modify();
+	Graph->RemoveNode(FoundComment, /*bBreakLinks=*/false);
+	Graph->MarkPackageDirty();
+}
diff --git a/Engine/Plugins/Experimental/Toolsets/DataflowAgent/Source/DataflowAgent/Public/DataflowAgentCustomTypes.h b/Engine/Plugins/Experimental/Toolsets/DataflowAgent/Source/DataflowAgent/Public/DataflowAgentCustomTypes.h
new file mode 100644
--- /dev/null
+++ b/Engine/Plugins/Experimental/Toolsets/DataflowAgent/Source/DataflowAgent/Public/DataflowAgentCustomTypes.h
@@ -0,0 +1,38 @@
+// Copyright Epic Games, Inc. All Rights Reserved.
+
+#pragma once
+
+#include "CoreMinimal.h"
+#include "DataflowAgentCustomTypes.generated.h"
+
+/** Basic info about a Dataflow node returned by GetNodeInfo */
+USTRUCT(BlueprintType)
+struct FDataflowNodeInfo
+{
+	GENERATED_BODY()
+
+	UPROPERTY(BlueprintReadOnly, Category = "Dataflow")
+	FString Name;
+
+	UPROPERTY(BlueprintReadOnly, Category = "Dataflow")
+	FString TypeName;
+
+	UPROPERTY(BlueprintReadOnly, Category = "Dataflow")
+	int32 PosX = 0;
+
+	UPROPERTY(BlueprintReadOnly, Category = "Dataflow")
+	int32 PosY = 0;
+};
+
+/** Describes one end of a connection in a Dataflow graph */
+USTRUCT(BlueprintType)
+struct FDataflowConnectionInfo
+{
+	GENERATED_BODY()
+
+	UPROPERTY(BlueprintReadOnly, Category = "Dataflow")
+	FString NodeName;
+
+	UPROPERTY(BlueprintReadOnly, Category = "Dataflow")
+	FString PinName;
+};
diff --git a/Engine/Plugins/Experimental/Toolsets/DataflowAgent/Source/DataflowAgent/Public/DataflowAgentModule.h b/Engine/Plugins/Experimental/Toolsets/DataflowAgent/Source/DataflowAgent/Public/DataflowAgentModule.h
new file mode 100644
--- /dev/null
+++ b/Engine/Plugins/Experimental/Toolsets/DataflowAgent/Source/DataflowAgent/Public/DataflowAgentModule.h
@@ -0,0 +1,15 @@
+// Copyright Epic Games, Inc. All Rights Reserved.
+
+#pragma once
+
+#include "Modules/ModuleManager.h"
+
+DECLARE_LOG_CATEGORY_EXTERN(LogDataflowAgent, Log, All);
+
+class FDataflowAgentModule : public IModuleInterface
+{
+public:
+	/** IModuleInterface implementation */
+	virtual void StartupModule() override;
+	virtual void ShutdownModule() override;
+};
diff --git a/Engine/Plugins/Experimental/Toolsets/DataflowAgent/Source/DataflowAgent/Public/DataflowAgentSkills.h b/Engine/Plugins/Experimental/Toolsets/DataflowAgent/Source/DataflowAgent/Public/DataflowAgentSkills.h
new file mode 100644
--- /dev/null
+++ b/Engine/Plugins/Experimental/Toolsets/DataflowAgent/Source/DataflowAgent/Public/DataflowAgentSkills.h
@@ -0,0 +1,37 @@
+// Copyright Epic Games, Inc. All Rights Reserved.
+
+#pragma once
+
+#include "DataflowAgentToolset.h"
+#include "ToolsetRegistry/AgentSkill.h"
+#include "ToolsetRegistry/ToolsetRegistry.h"
+
+#include "DataflowAgentSkills.generated.h"
+
+UCLASS()
+class UDataflowGraphEditingSkill : public UAgentSkill
+{
+	GENERATED_BODY()
+
+public:
+	UDataflowGraphEditingSkill()
+	{
+		Description = TEXT("Skill for generating and editing Dataflow graphs (simulation / deformation nodes)");
+		Toolsets.Add(UDataflowAgentToolset::StaticClass());
+	}
+
+protected:
+	virtual FString GeneratePrompt_Implementation(const FString& InitialInstructions) const override
+	{
+		// Build a simple prompt that lists available node types and forwards the initial instructions.
+		FString NodeList = UDataflowAgentToolset::ListNodeTypes(true);
+
+		FString Prompt;
+		Prompt += TEXT("You are an AI assistant helping the user build and edit Unreal Engine Dataflow graphs.\n\n");
+		Prompt += TEXT("Available node types:\n");
+		Prompt += NodeList;
+		Prompt += TEXT("\n\n");
+		Prompt += InitialInstructions;
+		return Prompt;
+	}
+};
diff --git a/Engine/Plugins/Experimental/Toolsets/DataflowAgent/Source/DataflowAgent/Public/DataflowAgentToolset.h b/Engine/Plugins/Experimental/Toolsets/DataflowAgent/Source/DataflowAgent/Public/DataflowAgentToolset.h
new file mode 100644
--- /dev/null
+++ b/Engine/Plugins/Experimental/Toolsets/DataflowAgent/Source/DataflowAgent/Public/DataflowAgentToolset.h
@@ -0,0 +1,233 @@
+// Copyright Epic Games, Inc. All Rights Reserved.
+
+#pragma once
+
+#include "DataflowAgentCustomTypes.h"
+#include "ToolsetRegistry/UToolsetRegistry.h"
+
+#include "DataflowAgentToolset.generated.h"
+
+class UDataflow;
+class UDataflowEdNode;
+
+/**
+ * Dataflow Agent Toolset - exposes Dataflow graph editing operations as AI agent tools.
+ */
+UCLASS(BlueprintType)
+class UDataflowAgentToolset : public UToolsetDefinition
+{
+	GENERATED_BODY()
+
+public:
+
+	// Graph -----------------------------------------------------------------------------------------------------------
+
+	/**
+	 * Creates a new saved Dataflow graph asset.
+	 *
+	 * @param Name - Name for the new Dataflow asset (e.g., "DF_Simulation")
+	 * @param Path - Content folder path where the asset should be created. Defaults to "/Game/Dataflow"
+	 * @return The full asset path of the created graph, or empty string on failure
+	 */
+	UFUNCTION(BlueprintCallable, Category = "Dataflow|Graph")
+	static FString CreateGraph(
+		const FString& Name,
+		const FString& Path = TEXT("/Game/Dataflow"));
+
+	/**
+	 * Returns the complete structure of a Dataflow graph including all nodes and connections.
+	 *
+	 * @param Graph - The Dataflow asset to inspect
+	 * @return JSON string with graph structure (nodes + connections), or error message on failure
+	 */
+	UFUNCTION(BlueprintCallable, Category = "Dataflow|Graph")
+	static FString GetGraphStructure(const UDataflow* Graph);
+
+	// Node Types ------------------------------------------------------------------------------------------------------
+
+	/**
+	 * Returns a JSON list of all registered Dataflow node types.
+	 *
+	 * @param bCommonOnly - When true, only return non-deprecated non-experimental nodes. Default is true.
+	 * @return JSON array of node type names
+	 */
+	UFUNCTION(BlueprintCallable, Category = "Dataflow|NodeTypes")
+	static FString ListNodeTypes(const bool bCommonOnly = true);
+
+	/**
+	 * Returns the schema for a Dataflow node type including its input/output pins
+	 * and editable UPROPERTY parameters.
+	 *
+	 * @param TypeName - The node type name (e.g., "FAddFloatsDataflowNode")
+	 * @return JSON object with name, category, tooltip, inputPins, outputPins, and properties
+	 */
+	UFUNCTION(BlueprintCallable, Category = "Dataflow|NodeTypes")
+	static FString GetNodeTypeSchema(const FString& TypeName);
+
+	// Nodes -----------------------------------------------------------------------------------------------------------
+
+	/**
+	 * Adds a node of the given type to the Dataflow graph.
+	 *
+	 * @param Graph      - The Dataflow asset to add the node to
+	 * @param TypeName   - The node type name (e.g., "FAddFloatsDataflowNode")
+	 * @param NodeName   - Unique name for the node within this graph
+	 * @param JsonParams - Optional JSON object of property overrides (e.g., {"Value": 3.14})
+	 * @param X          - X position in the graph editor (default 0, typical node width ~200)
+	 * @param Y          - Y position in the graph editor (default 0, typical node height ~100)
+	 * @return The newly created UDataflowEdNode, or nullptr on failure
+	 */
+	UFUNCTION(BlueprintCallable, Category = "Dataflow|Nodes")
+	static UDataflowEdNode* AddNode(
+		UDataflow* Graph,
+		const FString& TypeName,
+		const FString& NodeName,
+		const FString& JsonParams = TEXT(""),
+		const int32 X = 0,
+		const int32 Y = 0);
+
+	/**
+	 * Updates an existing node's editable properties via JSON.
+	 *
+	 * @param Node       - The EdNode to modify
+	 * @param JsonParams - JSON object of property overrides (e.g., {"Value": 3.14})
+	 */
+	UFUNCTION(BlueprintCallable, Category = "Dataflow|Nodes")
+	static void UpdateNode(UDataflowEdNode* Node, const FString& JsonParams);
+
+	/**
+	 * Returns information about a node as a JSON object (name, type, position, pins).
+	 *
+	 * @param Node - The EdNode to query
+	 * @return JSON string with node info, or error message on failure
+	 */
+	UFUNCTION(BlueprintCallable, Category = "Dataflow|Nodes")
+	static FString GetNodeInfo(UDataflowEdNode* Node);
+
+	/**
+	 * Moves a node to a new position in the graph editor.
+	 *
+	 * @param Node - The EdNode to reposition
+	 * @param X    - New X position
+	 * @param Y    - New Y position
+	 */
+	UFUNCTION(BlueprintCallable, Category = "Dataflow|Nodes")
+	static void RepositionNode(UDataflowEdNode* Node, const int32 X, const int32 Y);
+
+	/**
+	 * Removes a node and all its connections from the Dataflow graph.
+	 *
+	 * @param Graph - The Dataflow asset containing the node
+	 * @param Node  - The EdNode to remove
+	 */
+	UFUNCTION(BlueprintCallable, Category = "Dataflow|Nodes")
+	static void RemoveNode(UDataflow* Graph, UDataflowEdNode* Node);
+
+	// Connections -----------------------------------------------------------------------------------------------------
+
+	/**
+	 * Connects an output pin of one node to an input pin of another.
+	 *
+	 * @param FromNode - The source EdNode
+	 * @param FromPin  - Name of the output pin on the source node
+	 * @param ToNode   - The destination EdNode
+	 * @param ToPin    - Name of the input pin on the destination node
+	 * @return true on success, false on failure
+	 */
+	UFUNCTION(BlueprintCallable, Category = "Dataflow|Connections")
+	static bool ConnectNodePins(
+		UDataflowEdNode* FromNode, const FString& FromPin,
+		UDataflowEdNode* ToNode,   const FString& ToPin);
+
+	/**
+	 * Removes the connection between two node pins.
+	 *
+	 * @param FromNode - The source EdNode
+	 * @param FromPin  - Name of the output pin on the source node
+	 * @param ToNode   - The destination EdNode
+	 * @param ToPin    - Name of the input pin on the destination node
+	 */
+	UFUNCTION(BlueprintCallable, Category = "Dataflow|Connections")
+	static void DisconnectNodePins(
+		UDataflowEdNode* FromNode, const FString& FromPin,
+		UDataflowEdNode* ToNode,   const FString& ToPin);
+
+	// Variables -------------------------------------------------------------------------------------------------------
+
+	/**
+	 * Returns all variables defined on the Dataflow graph as a JSON array.
+	 * Each entry contains "name", "type", and "value" fields.
+	 *
+	 * @param Graph - The Dataflow asset to inspect
+	 * @return JSON array of variable descriptors, or empty string on failure
+	 */
+	UFUNCTION(BlueprintCallable, Category = "Dataflow|Variables")
+	static FString ListVariables(const UDataflow* Graph);
+
+	/**
+	 * Adds a new variable to the Dataflow graph.
+	 *
+	 * Supported type strings:
+	 *   Primitives : "Bool", "Int32", "Int64", "Float", "Double", "Name", "String"
+	 *   Structs    : UScriptStruct name with or without the "F" prefix
+	 *                e.g. "Vector", "FVector", "Transform", "FTransform", "Rotator", "LinearColor"
+	 *   Objects    : "Object:<ClassName>" where ClassName is with or without the "U"/"A" prefix
+	 *                e.g. "Object:StaticMesh", "Object:USkeletalMesh"
+	 *
+	 * @param Graph - The Dataflow asset to modify
+	 * @param Name  - Unique name for the new variable
+	 * @param Type  - Type string as described above
+	 * @return true on success, false if the variable already exists or the type cannot be resolved
+	 */
+	UFUNCTION(BlueprintCallable, Category = "Dataflow|Variables")
+	static bool AddVariable(UDataflow* Graph, const FString& Name, const FString& Type);
+
+	/**
+	 * Removes a variable from the Dataflow graph.
+	 *
+	 * @param Graph - The Dataflow asset to modify
+	 * @param Name  - Name of the variable to remove
+	 */
+	UFUNCTION(BlueprintCallable, Category = "Dataflow|Variables")
+	static void RemoveVariable(UDataflow* Graph, const FString& Name);
+
+	/**
+	 * Sets the value of an existing variable using its serialized string representation.
+	 * The format depends on the variable's type (e.g., "3.14" for float, "true" for bool,
+	 * "42" for int, "MyName" for FName).
+	 *
+	 * @param Graph - The Dataflow asset to modify
+	 * @param Name  - Name of the variable to set
+	 * @param Value - Serialized string value to assign
+	 * @return true on success, false if the variable is not found or the value cannot be parsed
+	 */
+	UFUNCTION(BlueprintCallable, Category = "Dataflow|Variables")
+	static bool SetVariable(UDataflow* Graph, const FString& Name, const FString& Value);
+
+	// Comments --------------------------------------------------------------------------------------------------------
+
+	/**
+	 * Adds a comment box around the given nodes.
+	 *
+	 * @param Graph   - The Dataflow asset to add the comment to
+	 * @param Nodes   - List of EdNodes to surround with the comment box
+	 * @param Comment - Text to display on the comment box
+	 * @param Color   - Background color of the comment box (defaults to White)
+	 * @return Node ID string of the created comment node, or empty string on failure
+	 */
+	UFUNCTION(BlueprintCallable, Category = "Dataflow|Comments")
+	static FString AddCommentBox(
+		UDataflow* Graph,
+		const TArray<UDataflowEdNode*>& Nodes,
+		const FString& Comment,
+		FLinearColor Color = FLinearColor::White);
+
+	/**
+	 * Removes a comment box node from the graph.
+	 *
+	 * @param Graph     - The Dataflow asset containing the comment
+	 * @param CommentId - The node GUID string returned by AddCommentBox
+	 */
+	UFUNCTION(BlueprintCallable, Category = "Dataflow|Comments")
+	static void RemoveCommentBox(UDataflow* Graph, const FString& CommentId);
+};
